Cleaned up runtime UFunction bindings on all property tracks.  Eliminated setting properties directly

diff --git a/Engine/Source/Runtime/MovieSceneCore/Classes/MovieScene.h b/Engine/Source/Runtime/MovieSceneCore/Classes/MovieScene.h
--- a/Engine/Source/Runtime/MovieSceneCore/Classes/MovieScene.h
+++ b/Engine/Source/Runtime/MovieSceneCore/Classes/MovieScene.h
@@ -4,6 +4,7 @@
 
 #include "MovieScene.generated.h"
 
+MOVIESCENECORE_API DECLARE_LOG_CATEGORY_EXTERN(LogSequencerRuntime, Log, All);
 
 /**
  * MovieSceneSpawnable describes an object that can be spawned for this MovieScene
diff --git a/Engine/Source/Runtime/MovieSceneCore/Private/MovieSceneCommonHelpers.cpp b/Engine/Source/Runtime/MovieSceneCore/Private/MovieSceneCommonHelpers.cpp
--- a/Engine/Source/Runtime/MovieSceneCore/Private/MovieSceneCommonHelpers.cpp
+++ b/Engine/Source/Runtime/MovieSceneCore/Private/MovieSceneCommonHelpers.cpp
@@ -72,3 +72,42 @@ UMovieSceneSection* MovieSceneHelpers::FindSectionAtTime( const TArray<UMovieSce
 	// if that's NULL, then there are no sections
 	return ClosestSection ? ClosestSection : EarliestSection;
 }
+
+
+FTrackInstancePropertyBindings::FTrackInstancePropertyBindings( FName PropertyName )
+{
+	static const FString Set(TEXT("Set"));
+
+	const FString FunctionString = Set + PropertyName.ToString();
+
+	FunctionName = FName(*FunctionString);
+}
+
+void FTrackInstancePropertyBindings::CallFunction( const TArray<UObject*>& InRuntimeObjects, void* FunctionParams )
+{
+	for( UObject* Object : InRuntimeObjects )
+	{
+		UFunction* Function = RuntimeObjectToFunctionMap.FindRef(Object);
+		if(Function)
+		{
+			Object->ProcessEvent(Function, FunctionParams);
+		}
+	}
+}
+
+void FTrackInstancePropertyBindings::UpdateBindings( const TArray<UObject*>& InRuntimeObjects )
+{
+	for(UObject* Object : InRuntimeObjects)
+	{
+		UFunction* Function = Object->FindFunction(FunctionName);
+		if(Function)
+		{
+			RuntimeObjectToFunctionMap.Add(Object, Function);
+		}
+		else
+		{
+			// Dont call potentially invalid functions
+			RuntimeObjectToFunctionMap.Remove(Object);
+		}
+	}
+}
diff --git a/Engine/Source/Runtime/MovieSceneCore/Private/MovieSceneCoreModule.cpp b/Engine/Source/Runtime/MovieSceneCore/Private/MovieSceneCoreModule.cpp
--- a/Engine/Source/Runtime/MovieSceneCore/Private/MovieSceneCoreModule.cpp
+++ b/Engine/Source/Runtime/MovieSceneCore/Private/MovieSceneCoreModule.cpp
@@ -7,6 +7,7 @@
 #include "MovieSceneBindings.h"
 #include "RuntimeMovieScenePlayer.h"
 
+DEFINE_LOG_CATEGORY(LogSequencerRuntime);
 
 
 /**
diff --git a/Engine/Source/Runtime/MovieSceneCore/Public/MovieSceneCommonHelpers.h b/Engine/Source/Runtime/MovieSceneCore/Public/MovieSceneCommonHelpers.h
--- a/Engine/Source/Runtime/MovieSceneCore/Public/MovieSceneCommonHelpers.h
+++ b/Engine/Source/Runtime/MovieSceneCore/Public/MovieSceneCommonHelpers.h
@@ -18,4 +18,22 @@ public:
 	* @return The found section or null
 	*/
 	static UMovieSceneSection* FindSectionAtTime( const TArray<UMovieSceneSection*>& Sections, float Time );
+
+};
+
+/**
+ * Manages bindings to keyed properties for a track instance. 
+ * Calls UFunctions to set the value on runtime objects
+ */
+class MOVIESCENECORE_API FTrackInstancePropertyBindings
+{
+public:
+	FTrackInstancePropertyBindings( FName PropertyName );
+
+	void CallFunction( const TArray<UObject*>& InRuntimeObjects, void* FunctionParams );
+	void UpdateBindings( const TArray<UObject*>& InRuntimeObjects );
+private:
+	/** Mapping of objects to bound functions that will be called to update data on the track */
+	TMap< TWeakObjectPtr<UObject>, UFunction* > RuntimeObjectToFunctionMap;
+	FName FunctionName;
 };
\ No newline at end of file
diff --git a/Engine/Source/Runtime/MovieSceneCoreTypes/Classes/MovieSceneVectorTrack.h b/Engine/Source/Runtime/MovieSceneCoreTypes/Classes/MovieSceneVectorTrack.h
--- a/Engine/Source/Runtime/MovieSceneCoreTypes/Classes/MovieSceneVectorTrack.h
+++ b/Engine/Source/Runtime/MovieSceneCoreTypes/Classes/MovieSceneVectorTrack.h
@@ -43,4 +43,11 @@ public:
 	 * @return true if anything was evaluated. Note: if false is returned OutVector remains unchanged
 	 */
 	virtual bool Eval( float Position, float LastPostion, FVector4& OutVector ) const;
+
+	/** @return Get the number of channels used by the vector */
+	int32 GetNumChannelsUsed() const { return NumChannelsUsed; }
+private:
+	/** The number of channels used by the vector (2,3, or 4)
+	UPROPERTY()
+	int32 NumChannelsUsed;
 };
diff --git a/Engine/Source/Runtime/MovieSceneCoreTypes/Private/Instances/MovieSceneBoolTrackInstance.cpp b/Engine/Source/Runtime/MovieSceneCoreTypes/Private/Instances/MovieSceneBoolTrackInstance.cpp
--- a/Engine/Source/Runtime/MovieSceneCoreTypes/Private/Instances/MovieSceneBoolTrackInstance.cpp
+++ b/Engine/Source/Runtime/MovieSceneCoreTypes/Private/Instances/MovieSceneBoolTrackInstance.cpp
@@ -2,12 +2,13 @@
 
 #include "MovieSceneCoreTypesPCH.h"
 #include "MovieSceneBoolTrackInstance.h"
-#include "MatineeUtils.h"
 
 
 FMovieSceneBoolTrackInstance::FMovieSceneBoolTrackInstance( UMovieSceneBoolTrack& InBoolTrack )
 {
 	BoolTrack = &InBoolTrack;
+	
+	PropertyBindings = MakeShareable( new FTrackInstancePropertyBindings( BoolTrack->GetPropertyName() ) );
 }
 
 void FMovieSceneBoolTrackInstance::Update( float Position, float LastPosition, const TArray<UObject*>& RuntimeObjects, class IMovieScenePlayer& Player ) 
@@ -15,21 +16,12 @@ void FMovieSceneBoolTrackInstance::Update( float Position, float LastPosition, c
 	bool BoolValue = false;
 	if( BoolTrack->Eval( Position, LastPosition, BoolValue ) )
 	{
-		for( int32 ObjIndex = 0; ObjIndex < RuntimeObjects.Num(); ++ObjIndex )
-		{
-			UObject* Object = RuntimeObjects[ObjIndex];
-
-			UObject* PropertyOwner = NULL;
-			UProperty* Property = NULL;
-			//@todo Sequencer - Major performance problems here.  This needs to be initialized and stored (not serialized) somewhere
-			uint8* Address = FMatineeUtils::GetPropertyAddress<uint8>( Object, BoolTrack->GetPropertyName(), Property, PropertyOwner );
-			if( Address )
-			{
-				CastChecked<UBoolProperty>(Property)->SetPropertyValue(Address, BoolValue);
-				// Let the property owner know that we changed one of its properties
-				PropertyOwner->PostInterpChange( Property );
-			}
-		}
+		PropertyBindings->CallFunction( RuntimeObjects, &BoolValue );
 	}
 }
 
+void FMovieSceneBoolTrackInstance::RefreshInstance( const TArray<UObject*>& RuntimeObjects, class IMovieScenePlayer& Player )
+{
+	PropertyBindings->UpdateBindings( RuntimeObjects );
+}
+
diff --git a/Engine/Source/Runtime/MovieSceneCoreTypes/Private/Instances/MovieSceneBoolTrackInstance.h b/Engine/Source/Runtime/MovieSceneCoreTypes/Private/Instances/MovieSceneBoolTrackInstance.h
--- a/Engine/Source/Runtime/MovieSceneCoreTypes/Private/Instances/MovieSceneBoolTrackInstance.h
+++ b/Engine/Source/Runtime/MovieSceneCoreTypes/Private/Instances/MovieSceneBoolTrackInstance.h
@@ -12,8 +12,10 @@ public:
 
 	/** IMovieSceneTrackInstance interface */
 	virtual void Update( float Position, float LastPosition, const TArray<UObject*>& RuntimeObjects, class IMovieScenePlayer& Player ) override;
-	virtual void RefreshInstance( const TArray<UObject*>& RuntimeObjects, class IMovieScenePlayer& Player ) override {}
+	virtual void RefreshInstance( const TArray<UObject*>& RuntimeObjects, class IMovieScenePlayer& Player ) override;
 private:
 	/** Bool track that is being instanced */
 	UMovieSceneBoolTrack* BoolTrack;
+	/** Runtime property bindings */
+	TSharedPtr<class FTrackInstancePropertyBindings> PropertyBindings;
 };
diff --git a/Engine/Source/Runtime/MovieSceneCoreTypes/Private/Instances/MovieSceneColorTrackInstance.cpp b/Engine/Source/Runtime/MovieSceneCoreTypes/Private/Instances/MovieSceneColorTrackInstance.cpp
--- a/Engine/Source/Runtime/MovieSceneCoreTypes/Private/Instances/MovieSceneColorTrackInstance.cpp
+++ b/Engine/Source/Runtime/MovieSceneCoreTypes/Private/Instances/MovieSceneColorTrackInstance.cpp
@@ -2,12 +2,14 @@
 
 #include "MovieSceneCoreTypesPCH.h"
 #include "MovieSceneColorTrackInstance.h"
-#include "MatineeUtils.h"
+#include "MovieSceneCommonHelpers.h"
 
 
 FMovieSceneColorTrackInstance::FMovieSceneColorTrackInstance( UMovieSceneColorTrack& InColorTrack )
 {
 	ColorTrack = &InColorTrack;
+
+	PropertyBindings = MakeShareable( new FTrackInstancePropertyBindings( ColorTrack->GetPropertyName() ) );
 }
 
 void FMovieSceneColorTrackInstance::Update( float Position, float LastPosition, const TArray<UObject*>& RuntimeObjects, class IMovieScenePlayer& Player ) 
@@ -15,72 +17,11 @@ void FMovieSceneColorTrackInstance::Update( float Position, float LastPosition,
 	FLinearColor ColorValue;
 	if( ColorTrack->Eval( Position, LastPosition, ColorValue ) )
 	{
-		for( int32 ObjIndex = 0; ObjIndex < RuntimeObjects.Num(); ++ObjIndex )
-		{
-			UObject* Object = RuntimeObjects[ObjIndex];
-
-			UFunction* Function = RuntimeObjectToFunctionMap.FindRef( Object );
-			if( Function )
-			{
-				Object->ProcessEvent( Function, &ColorValue );
-			}
-			else
-			{
-				// @todo Sequencer - remove this legacy method of setting the property directly
-				/*UObject* PropertyOwner = NULL;
-				UProperty* Property = NULL;
-				//@todo Sequencer - Major performance problems here.  This needs to be initialized and stored (not serialized) somewhere
-				void* Address = FMatineeUtils::GetPropertyAddress<void>( Object, ColorTrack->GetPropertyName(), Property, PropertyOwner );
-
-				if (Address)
-				{
-					bool bIsFColor = false;
-					const UStructProperty* StructProp = Cast<const UStructProperty>(Property);
-					check(StructProp && StructProp->Struct);
-					FName StructName = StructProp->Struct->GetFName();
-
-					bIsFColor = StructName == NAME_Color;
-					check(bIsFColor || StructName == NAME_LinearColor);
-
-					if( bIsFColor )
-					{
-						*(FColor*)Address = ColorValue.ToFColor(false);
-						// Let the property owner know that we changed one of its properties
-						PropertyOwner->PostInterpChange( Property );
-					}
-					else
-					{
-						*(FLinearColor*)Address = ColorValue;
-						// Let the property owner know that we changed one of its properties
-						PropertyOwner->PostInterpChange( Property );
-					}
-				}*/
-			}
-		}
+		PropertyBindings->CallFunction( RuntimeObjects, &ColorValue );
 	}
 }
 
 void FMovieSceneColorTrackInstance::RefreshInstance( const TArray<UObject*>& RuntimeObjects, class IMovieScenePlayer& Player )
 {
-	static const FString Set(TEXT("Set") );
-	
-	FString FunctionString = Set + ColorTrack->GetPropertyName().ToString();
-	
-	FName FunctionName = FName( *FunctionString );
-	
-	//Player.GetRuntimeObjects( )
-	for( UObject* Object : RuntimeObjects )
-	{
-		// If the mapping for the object already exists we dont need to search for it again
-		UFunction* Function = RuntimeObjectToFunctionMap.FindRef( Object );
-		
-		if( !Function )
-		{
-			Function = Object->FindFunction( FunctionName );
-			if( Function )
-			{
-				RuntimeObjectToFunctionMap.Add( Object, Function );
-			}
-		}
-	}
+	PropertyBindings->UpdateBindings( RuntimeObjects );
 }
diff --git a/Engine/Source/Runtime/MovieSceneCoreTypes/Private/Instances/MovieSceneColorTrackInstance.h b/Engine/Source/Runtime/MovieSceneCoreTypes/Private/Instances/MovieSceneColorTrackInstance.h
--- a/Engine/Source/Runtime/MovieSceneCoreTypes/Private/Instances/MovieSceneColorTrackInstance.h
+++ b/Engine/Source/Runtime/MovieSceneCoreTypes/Private/Instances/MovieSceneColorTrackInstance.h
@@ -16,6 +16,6 @@ public:
 private:
 	/** The track being instanced */
 	UMovieSceneColorTrack* ColorTrack;
-	/** Mapping of objects to bound functions that will be called to update data on the track */
-	TMap< TWeakObjectPtr<UObject>, UFunction* > RuntimeObjectToFunctionMap;
+	/** Runtime property bindings */
+	TSharedPtr<class FTrackInstancePropertyBindings> PropertyBindings;
 };
diff --git a/Engine/Source/Runtime/MovieSceneCoreTypes/Private/Instances/MovieSceneFloatTrackInstance.cpp b/Engine/Source/Runtime/MovieSceneCoreTypes/Private/Instances/MovieSceneFloatTrackInstance.cpp
--- a/Engine/Source/Runtime/MovieSceneCoreTypes/Private/Instances/MovieSceneFloatTrackInstance.cpp
+++ b/Engine/Source/Runtime/MovieSceneCoreTypes/Private/Instances/MovieSceneFloatTrackInstance.cpp
@@ -2,12 +2,13 @@
 
 #include "MovieSceneCoreTypesPCH.h"
 #include "MovieSceneFloatTrackInstance.h"
-#include "MatineeUtils.h"
 
 
 FMovieSceneFloatTrackInstance::FMovieSceneFloatTrackInstance( UMovieSceneFloatTrack& InFloatTrack )
 {
 	FloatTrack = &InFloatTrack;
+
+	PropertyBindings = MakeShareable( new FTrackInstancePropertyBindings( FloatTrack->GetPropertyName() ) );
 }
 
 void FMovieSceneFloatTrackInstance::Update( float Position, float LastPosition, const TArray<UObject*>& RuntimeObjects, class IMovieScenePlayer& Player ) 
@@ -15,21 +16,11 @@ void FMovieSceneFloatTrackInstance::Update( float Position, float LastPosition,
 	float FloatValue = 0.0f;
 	if( FloatTrack->Eval( Position, LastPosition, FloatValue ) )
 	{
-		for( int32 ObjIndex = 0; ObjIndex < RuntimeObjects.Num(); ++ObjIndex )
-		{
-			UObject* Object = RuntimeObjects[ObjIndex];
-
-			UObject* PropertyOwner = NULL;
-			UProperty* Property = NULL;
-			//@todo Sequencer - Major performance problems here.  This needs to be initialized and stored (not serialized) somewhere
-			float* Address = FMatineeUtils::GetPropertyAddress<float>( Object, FloatTrack->GetPropertyName(), Property, PropertyOwner );
-			if( Address )
-			{
-				*Address = FloatValue;
-				// Let the property owner know that we changed one of its properties
-				PropertyOwner->PostInterpChange( Property );
-			}
-		}
+		PropertyBindings->CallFunction( RuntimeObjects, &FloatValue );
 	}
 }
 
+void FMovieSceneFloatTrackInstance::RefreshInstance( const TArray<UObject*>& RuntimeObjects, class IMovieScenePlayer& Player )
+{
+	PropertyBindings->UpdateBindings( RuntimeObjects );
+}
diff --git a/Engine/Source/Runtime/MovieSceneCoreTypes/Private/Instances/MovieSceneFloatTrackInstance.h b/Engine/Source/Runtime/MovieSceneCoreTypes/Private/Instances/MovieSceneFloatTrackInstance.h
--- a/Engine/Source/Runtime/MovieSceneCoreTypes/Private/Instances/MovieSceneFloatTrackInstance.h
+++ b/Engine/Source/Runtime/MovieSceneCoreTypes/Private/Instances/MovieSceneFloatTrackInstance.h
@@ -12,8 +12,10 @@ public:
 
 	/** IMovieSceneTrackInstance interface */
 	virtual void Update( float Position, float LastPosition, const TArray<UObject*>& RuntimeObjects, class IMovieScenePlayer& Player ) override;
-	virtual void RefreshInstance( const TArray<UObject*>& RuntimeObjects, class IMovieScenePlayer& Player ) override {}
+	virtual void RefreshInstance( const TArray<UObject*>& RuntimeObjects, class IMovieScenePlayer& Player ) override;
 private:
 	/** Track that is being instanced */
 	UMovieSceneFloatTrack* FloatTrack;
+	/** Runtime property bindings */
+	TSharedPtr<class FTrackInstancePropertyBindings> PropertyBindings;
 };
diff --git a/Engine/Source/Runtime/MovieSceneCoreTypes/Private/Instances/MovieSceneVectorTrackInstance.cpp b/Engine/Source/Runtime/MovieSceneCoreTypes/Private/Instances/MovieSceneVectorTrackInstance.cpp
--- a/Engine/Source/Runtime/MovieSceneCoreTypes/Private/Instances/MovieSceneVectorTrackInstance.cpp
+++ b/Engine/Source/Runtime/MovieSceneCoreTypes/Private/Instances/MovieSceneVectorTrackInstance.cpp
@@ -2,12 +2,12 @@
 
 #include "MovieSceneCoreTypesPCH.h"
 #include "MovieSceneVectorTrackInstance.h"
-#include "MatineeUtils.h"
-
 
 FMovieSceneVectorTrackInstance::FMovieSceneVectorTrackInstance( UMovieSceneVectorTrack& InVectorTrack )
 {
 	VectorTrack = &InVectorTrack;
+
+	PropertyBindings = MakeShareable( new FTrackInstancePropertyBindings( VectorTrack->GetPropertyName() ) );
 }
 
 void FMovieSceneVectorTrackInstance::Update( float Position, float LastPosition, const TArray<UObject*>& RuntimeObjects, class IMovieScenePlayer& Player ) 
@@ -15,43 +15,36 @@ void FMovieSceneVectorTrackInstance::Update( float Position, float LastPosition,
 	FVector4 Vector;
 	if( VectorTrack->Eval( Position, LastPosition, Vector ) )
 	{
-		for( int32 ObjIndex = 0; ObjIndex < RuntimeObjects.Num(); ++ObjIndex )
+		int32 NumChannelsUsed = VectorTrack->GetNumChannelsUsed();
+		switch( NumChannelsUsed )
 		{
-			UObject* Object = RuntimeObjects[ObjIndex];
-
-			UObject* PropertyOwner = NULL;
-			UProperty* Property = NULL;
-			//@todo Sequencer - Major performance problems here.  This needs to be initialized and stored (not serialized) somewhere
-			void* Address = FMatineeUtils::GetPropertyAddress<void>( Object, VectorTrack->GetPropertyName(), Property, PropertyOwner );
-		
-			if (Address)
+			case 2:
 			{
-				bool bIsVector2D = false, bIsVector = false, bIsVector4 = false;
-				const UStructProperty* StructProp = Cast<const UStructProperty>(Property);
-				check(StructProp && StructProp->Struct);
-				FName StructName = StructProp->Struct->GetFName();
-
-				bIsVector2D = StructName == NAME_Vector2D;
-				bIsVector = StructName == NAME_Vector;
-				bIsVector4 = StructName == NAME_Vector4;
-				check(bIsVector2D || bIsVector || bIsVector4);
-
-				if (bIsVector2D)
-				{
-					*(FVector2D*)Address = FVector2D(Vector.X, Vector.Y);
-				}
-				else if (bIsVector)
-				{
-					*(FVector*)Address = FVector(Vector.X, Vector.Y, Vector.Z);
-				}
-				else if (bIsVector4)
-				{
-					*(FVector4*)Address = Vector;
-				}
-				// Let the property owner know that we changed one of its properties
-				PropertyOwner->PostInterpChange( Property );
+				FVector2D Value(Vector.X, Vector.Y);
+				PropertyBindings->CallFunction(RuntimeObjects, &Value);
+				break;
+			}
+			case 3:
+			{
+				FVector Value(Vector.X, Vector.Y, Vector.Z);
+				PropertyBindings->CallFunction(RuntimeObjects, &Value);
+				break;
+			}
+			case 4:
+			{
+				PropertyBindings->CallFunction(RuntimeObjects, &Vector);
+				break;
 			}
+			default:
+				UE_LOG(LogSequencerRuntime, Warning, TEXT("Invalid number of channels(%d) for vector track"), NumChannelsUsed );
+				break;
 		}
+		
 	}
 }
 
+
+void FMovieSceneVectorTrackInstance::RefreshInstance(const TArray<UObject*>& RuntimeObjects, class IMovieScenePlayer& Player)
+{
+	PropertyBindings->UpdateBindings(RuntimeObjects);
+}
diff --git a/Engine/Source/Runtime/MovieSceneCoreTypes/Private/Instances/MovieSceneVectorTrackInstance.h b/Engine/Source/Runtime/MovieSceneCoreTypes/Private/Instances/MovieSceneVectorTrackInstance.h
--- a/Engine/Source/Runtime/MovieSceneCoreTypes/Private/Instances/MovieSceneVectorTrackInstance.h
+++ b/Engine/Source/Runtime/MovieSceneCoreTypes/Private/Instances/MovieSceneVectorTrackInstance.h
@@ -12,8 +12,10 @@ public:
 
 	/** IMovieSceneTrackInstance interface */
 	virtual void Update( float Position, float LastPosition, const TArray<UObject*>& RuntimeObjects, class IMovieScenePlayer& Player ) override;
-	virtual void RefreshInstance( const TArray<UObject*>& RuntimeObjects, class IMovieScenePlayer& Player ) override {}
+	virtual void RefreshInstance( const TArray<UObject*>& RuntimeObjects, class IMovieScenePlayer& Player ) override;
 private:
 	/** Track that is being instanced */
 	UMovieSceneVectorTrack* VectorTrack;
+	/** Runtime property bindings */
+	TSharedPtr<class FTrackInstancePropertyBindings> PropertyBindings;
 };
diff --git a/Engine/Source/Runtime/MovieSceneCoreTypes/Private/MovieSceneVectorTrack.cpp b/Engine/Source/Runtime/MovieSceneCoreTypes/Private/MovieSceneVectorTrack.cpp
--- a/Engine/Source/Runtime/MovieSceneCoreTypes/Private/MovieSceneVectorTrack.cpp
+++ b/Engine/Source/Runtime/MovieSceneCoreTypes/Private/MovieSceneVectorTrack.cpp
@@ -9,6 +9,7 @@
 UMovieSceneVectorTrack::UMovieSceneVectorTrack( const FPostConstructInitializeProperties& PCIP )
 	: Super( PCIP )
 {
+	NumChannelsUsed = 0;
 }
 
 UMovieSceneSection* UMovieSceneVectorTrack::CreateNewSection()
@@ -32,6 +33,9 @@ bool UMovieSceneVectorTrack::AddKeyToSection( float Time, const FVector4& Value,
 
 		NewSection->AddKey( Time, Value );
 
+		// We dont support one track containing multiple sections of differing channel amounts
+		NumChannelsUsed = InChannelsUsed;
+
 		return true;
 	}
 	return false;
diff --git a/Engine/Source/Runtime/UMG/Private/Animation/MovieSceneMarginTrackInstance.cpp b/Engine/Source/Runtime/UMG/Private/Animation/MovieSceneMarginTrackInstance.cpp
--- a/Engine/Source/Runtime/UMG/Private/Animation/MovieSceneMarginTrackInstance.cpp
+++ b/Engine/Source/Runtime/UMG/Private/Animation/MovieSceneMarginTrackInstance.cpp
@@ -6,6 +6,8 @@
 FMovieSceneMarginTrackInstance::FMovieSceneMarginTrackInstance( UMovieSceneMarginTrack& InMarginTrack )
 {
 	MarginTrack = &InMarginTrack;
+
+	PropertyBindings = MakeShareable( new FTrackInstancePropertyBindings( MarginTrack->GetPropertyName() ) );
 }
 
 void FMovieSceneMarginTrackInstance::Update( float Position, float LastPosition, const TArray<UObject*>& RuntimeObjects, class IMovieScenePlayer& Player ) 
@@ -13,40 +15,11 @@ void FMovieSceneMarginTrackInstance::Update( float Position, float LastPosition,
 	FMargin MarginValue;
 	if( MarginTrack->Eval( Position, LastPosition, MarginValue ) )
 	{
-		for( int32 ObjIndex = 0; ObjIndex < RuntimeObjects.Num(); ++ObjIndex )
-		{
-			UObject* Object = RuntimeObjects[ObjIndex];
-
-			UFunction* Function = RuntimeObjectToFunctionMap.FindRef( Object );
-			if( Function )
-			{
-				Object->ProcessEvent( Function, &MarginValue );
-			}
-		}
+		PropertyBindings->CallFunction( RuntimeObjects, &MarginValue );
 	}
 }
 
 void FMovieSceneMarginTrackInstance::RefreshInstance( const TArray<UObject*>& RuntimeObjects, class IMovieScenePlayer& Player )
 {
-	static const FString Set(TEXT("Set") );
-	
-	FString FunctionString = Set + MarginTrack->GetPropertyName().ToString();
-	
-	FName FunctionName = FName( *FunctionString );
-	
-	//Player.GetRuntimeObjects( )
-	for( UObject* Object : RuntimeObjects )
-	{
-		// If the mapping for the object already exists we dont need to search for it again
-		UFunction* Function = RuntimeObjectToFunctionMap.FindRef( Object );
-		
-		if( !Function )
-		{
-			Function = Object->FindFunction( FunctionName );
-			if( Function )
-			{
-				RuntimeObjectToFunctionMap.Add( Object, Function );
-			}
-		}
-	}
+	PropertyBindings->UpdateBindings( RuntimeObjects );
 }
diff --git a/Engine/Source/Runtime/UMG/Private/Animation/MovieSceneMarginTrackInstance.h b/Engine/Source/Runtime/UMG/Private/Animation/MovieSceneMarginTrackInstance.h
--- a/Engine/Source/Runtime/UMG/Private/Animation/MovieSceneMarginTrackInstance.h
+++ b/Engine/Source/Runtime/UMG/Private/Animation/MovieSceneMarginTrackInstance.h
@@ -5,7 +5,7 @@
 #include "IMovieSceneTrackInstance.h"
 
 /**
- * Instance of a UMovieSceneColorTrack
+ * Instance of a UMovieSceneMarginTrack
  */
 class FMovieSceneMarginTrackInstance : public IMovieSceneTrackInstance
 {
@@ -18,6 +18,6 @@ public:
 private:
 	/** The track being instanced */
 	UMovieSceneMarginTrack* MarginTrack;
-	/** Mapping of objects to bound functions that will be called to update data on the track */
-	TMap< TWeakObjectPtr<UObject>, UFunction* > RuntimeObjectToFunctionMap;
+	/** Runtime property bindings */
+	TSharedPtr<class FTrackInstancePropertyBindings> PropertyBindings;
 };
