Fixed missing end comment marker

diff --git a/Engine/Source/Runtime/MovieSceneCoreTypes/Classes/MovieSceneVectorTrack.h b/Engine/Source/Runtime/MovieSceneCoreTypes/Classes/MovieSceneVectorTrack.h
--- a/Engine/Source/Runtime/MovieSceneCoreTypes/Classes/MovieSceneVectorTrack.h
+++ b/Engine/Source/Runtime/MovieSceneCoreTypes/Classes/MovieSceneVectorTrack.h
@@ -47,7 +47,7 @@ public:
 	/** @return Get the number of channels used by the vector */
 	int32 GetNumChannelsUsed() const { return NumChannelsUsed; }
 private:
-	/** The number of channels used by the vector (2,3, or 4)
+	/** The number of channels used by the vector (2,3, or 4) */
 	UPROPERTY()
 	int32 NumChannelsUsed;
 };
