Added missing file

diff --git a/Engine/Source/Runtime/MovieSceneCoreTypes/Classes/MovieScene3DTransformSection.h b/Engine/Source/Runtime/MovieSceneCoreTypes/Classes/MovieScene3DTransformSection.h
--- a/Engine/Source/Runtime/MovieSceneCoreTypes/Classes/MovieScene3DTransformSection.h
+++ b/Engine/Source/Runtime/MovieSceneCoreTypes/Classes/MovieScene3DTransformSection.h
@@ -52,7 +52,7 @@ public:
 	 * @param Axis	The axis of the translation curve to get
 	 * @return The curve on the axis
 	 */
-	virtual FRichCurve& GetTranslationCurve( EAxis::Type Axis );
+	MOVIESCENECORETYPES_API FRichCurve& GetTranslationCurve( EAxis::Type Axis );
 
 	/** 
 	 * Returns the rotation curve for a specific axis
@@ -60,7 +60,7 @@ public:
 	 * @param Axis	The axis of the rotation curve to get
 	 * @return The curve on the axis
 	 */
-	virtual FRichCurve& GetRotationCurve( EAxis::Type Axis );
+	MOVIESCENECORETYPES_API FRichCurve& GetRotationCurve( EAxis::Type Axis );
 
 	/** 
 	 * Returns the scale curve for a specific axis
@@ -68,7 +68,7 @@ public:
 	 * @param Axis	The axis of the scale curve to get
 	 * @return The curve on the axis
 	 */
-	virtual FRichCurve& GetScaleCurve( EAxis::Type Axis );
+	MOVIESCENECORETYPES_API FRichCurve& GetScaleCurve( EAxis::Type Axis );
 
 	/**
 	 * Adds keys each component of translation. Note: Only adds keys if a value has changed
@@ -99,16 +99,6 @@ public:
 	 * @return True if the new key would be new data, false if duplicate
 	 */
 	bool NewKeyIsNewData(const class FTransformKey& TransformKey) const;
-private:
-	/**
-	 * Adds a key to a rich curve, finding an existing key to modify or adding a new one
-	 *
-	 * @param InCurve	The curve to add keys to
-	 * @param Time		The time where the key should be added
-	 * @param Value		The value at the given time
-	 * @param	bUnwindRotation		When true, the value will be treated like a rotation value in degrees, and will automatically be unwound to prevent flipping 360 degrees from the previous key 
-	 */
-	void AddKeyToCurve( FRichCurve& InCurve, float Time, float Value, const bool bUnwindRotation );
 private:
 	/** Translation curves */
 	UPROPERTY()
