[INTEGRATE] Change 2419928 by [REDACTED_PERSON][REDACTED_HANDLE] on 2015/01/26 23:52:38

	Fix some includes for the overnight builds.

diff --git a/Engine/Source/Editor/UnrealEd/Classes/Settings/LevelEditorPlaySettings.h b/Engine/Source/Editor/UnrealEd/Classes/Settings/LevelEditorPlaySettings.h
--- a/Engine/Source/Editor/UnrealEd/Classes/Settings/LevelEditorPlaySettings.h
+++ b/Engine/Source/Editor/UnrealEd/Classes/Settings/LevelEditorPlaySettings.h
@@ -2,7 +2,7 @@
 
 #pragma once
 
-
+#include "SlateCore.h"
 #include "LevelEditorPlaySettings.generated.h"
 
 
diff --git a/Engine/Source/Editor/UnrealEd/Private/SComponentClassCombo.cpp b/Engine/Source/Editor/UnrealEd/Private/SComponentClassCombo.cpp
--- a/Engine/Source/Editor/UnrealEd/Private/SComponentClassCombo.cpp
+++ b/Engine/Source/Editor/UnrealEd/Private/SComponentClassCombo.cpp
@@ -12,6 +12,7 @@
 #include "SSearchBox.h"
 #include "Engine/Selection.h"
 #include "HotReloadInterface.h"
+#include "KismetEditorUtilities.h"
 
 #define LOCTEXT_NAMESPACE "ComponentClassCombo"
 
diff --git a/Engine/Source/Runtime/Engine/Private/KismetSystemLibrary.cpp b/Engine/Source/Runtime/Engine/Private/KismetSystemLibrary.cpp
--- a/Engine/Source/Runtime/Engine/Private/KismetSystemLibrary.cpp
+++ b/Engine/Source/Runtime/Engine/Private/KismetSystemLibrary.cpp
@@ -12,6 +12,7 @@
 #include "Camera/CameraActor.h"
 #include "Camera/CameraComponent.h"
 #include "Kismet/GameplayStatics.h"
+#include "SlateCore.h"
 
 //////////////////////////////////////////////////////////////////////////
 // UKismetSystemLibrary
diff --git a/Engine/Source/Runtime/Engine/Public/EngineFontServices.h b/Engine/Source/Runtime/Engine/Public/EngineFontServices.h
--- a/Engine/Source/Runtime/Engine/Public/EngineFontServices.h
+++ b/Engine/Source/Runtime/Engine/Public/EngineFontServices.h
@@ -2,6 +2,8 @@
 
 #pragma once
 
+#include "SlateCore.h"
+
 /** 
  * Provides the Engine with access to the Slate font cache and font measuring services (for Canvas)
  * On the game thread this just leverages the Slate services, but the render thread needs its own instances
diff --git a/Engine/Source/Runtime/MovieSceneCoreTypes/Private/Instances/MovieSceneColorTrackInstance.cpp b/Engine/Source/Runtime/MovieSceneCoreTypes/Private/Instances/MovieSceneColorTrackInstance.cpp
--- a/Engine/Source/Runtime/MovieSceneCoreTypes/Private/Instances/MovieSceneColorTrackInstance.cpp
+++ b/Engine/Source/Runtime/MovieSceneCoreTypes/Private/Instances/MovieSceneColorTrackInstance.cpp
@@ -3,6 +3,7 @@
 #include "MovieSceneCoreTypesPCH.h"
 #include "MovieSceneColorTrackInstance.h"
 #include "MovieSceneCommonHelpers.h"
+#include "SlateCore.h"
 
 
 FMovieSceneColorTrackInstance::FMovieSceneColorTrackInstance( UMovieSceneColorTrack& InColorTrack )
