Sequencer: Get audio playing.

diff --git a/Engine/Source/Runtime/MovieSceneCoreTypes/Private/Instances/MovieSceneAudioTrackInstance.cpp b/Engine/Source/Runtime/MovieSceneCoreTypes/Private/Instances/MovieSceneAudioTrackInstance.cpp
--- a/Engine/Source/Runtime/MovieSceneCoreTypes/Private/Instances/MovieSceneAudioTrackInstance.cpp
+++ b/Engine/Source/Runtime/MovieSceneCoreTypes/Private/Instances/MovieSceneAudioTrackInstance.cpp
@@ -125,6 +125,7 @@ void FMovieSceneAudioTrackInstance::PlaySound(UMovieSceneAudioSection* AudioSect
 	Component->SetSound(AudioSection->GetSound());
 	Component->SetVolumeMultiplier(1.f);
 	Component->SetPitchMultiplier(PitchMultiplier);
+	Component->bIsUISound = true;
 	Component->Play(Time - AudioSection->GetAudioStartTime());
 }
 
