Sequencer: Save/restore state when opening/closing sequencer and when saving the level.

An active sequencer pointer is stored when initializing a sequencer. That active sequencer is shut down if another sequencer is initialized.

Tested:
1) Object returns to original position after closing the sequencer window.
2) Object remains at animated position when saving the level, and saved level opens with object in original position.
3) Object retains original position for saving when switching from one movie scene to another, but takes animated position of active movie scene.
4) Works with UMG and Sequencer tabs opened/closed.

diff --git a/Engine/Source/Editor/Sequencer/Private/Sequencer.cpp b/Engine/Source/Editor/Sequencer/Private/Sequencer.cpp
--- a/Engine/Source/Editor/Sequencer/Private/Sequencer.cpp
+++ b/Engine/Source/Editor/Sequencer/Private/Sequencer.cpp
@@ -47,6 +47,10 @@ bool FSequencer::IsSequencerEnabled()
 	return true;
 }
 
+namespace
+{
+	TSharedPtr<FSequencer> ActiveSequencer;
+}
 
 void FSequencer::InitSequencer( const FSequencerInitParams& InitParams, const TArray<FOnCreateTrackEditor>& TrackEditorDelegates )
 {
@@ -54,6 +58,20 @@ void FSequencer::InitSequencer( const FSequencerInitParams& InitParams, const TA
 	{
 		bIsEditingWithinLevelEditor = InitParams.bEditWithinLevelEditor;
 
+		// If this sequencer edits the level, close out the existing active sequencer and mark this as the active sequencer.
+		if (bIsEditingWithinLevelEditor)
+		{
+			if (ActiveSequencer.IsValid())
+			{
+				ActiveSequencer->OnClose();
+			}
+			ActiveSequencer = SharedThis(this);
+
+			// Register for saving the level so that the state of the scene can be restored before saving and updated after saving.
+			FEditorDelegates::PreSaveWorld.AddSP(this, &FSequencer::OnPreSaveWorld);
+			FEditorDelegates::PostSaveWorld.AddSP(this, &FSequencer::OnPostSaveWorld);
+		}
+
 		ToolkitHost = InitParams.ToolkitHost;
 
 		LastViewRange = TargetViewRange = InitParams.ViewParams.InitalViewRange;
@@ -156,6 +174,20 @@ FSequencer::~FSequencer()
 	SequencerWidget.Reset();
 }
 
+void
+FSequencer::OnClose()
+{
+	if (ActiveSequencer.IsValid() && ActiveSequencer.Get() == this)
+	{
+		RootMovieSceneInstance->RestoreState();
+
+		FEditorDelegates::PreSaveWorld.RemoveAll(this);
+		FEditorDelegates::PostSaveWorld.RemoveAll(this);
+
+		ActiveSequencer.Reset();
+	}
+}
+
 void FSequencer::Tick(float InDeltaTime)
 {
 	if (bNeedTreeRefresh)
@@ -1172,6 +1204,17 @@ void FSequencer::PostUndo(bool bSuccess)
 	NotifyMovieSceneDataChanged();
 }
 
+void FSequencer::OnPreSaveWorld(uint32 SaveFlags, class UWorld* World)
+{
+	// Restore the saved state so that the level save can save that instead of the animated state.
+	RootMovieSceneInstance->RestoreState();
+}
+
+void FSequencer::OnPostSaveWorld(uint32 SaveFlags, class UWorld* World, bool bSuccess)
+{
+	// Reset the time after saving so that an update will be triggered to put objects back to their animated state.
+	SetGlobalTime(GetGlobalTime());
+}
 
 void FSequencer::OnSectionSelectionChanged()
 {
diff --git a/Engine/Source/Editor/Sequencer/Private/Sequencer.h b/Engine/Source/Editor/Sequencer/Private/Sequencer.h
--- a/Engine/Source/Editor/Sequencer/Private/Sequencer.h
+++ b/Engine/Source/Editor/Sequencer/Private/Sequencer.h
@@ -49,6 +49,11 @@ public:
 	 */
 	void InitSequencer( const FSequencerInitParams& InitParams, const TArray<FOnCreateTrackEditor>& TrackEditorDelegates );
 
+	/** 
+	* Closes the sequencer
+	*/
+	void OnClose();
+
 	/** Constructor */
 	FSequencer();
 
@@ -341,6 +346,12 @@ protected:
 
 	void OnSectionSelectionChanged();
 
+	/** Called before the world is going to be saved. The sequencer puts everything back to its initial state. */
+	void OnPreSaveWorld(uint32 SaveFlags, class UWorld* World);
+
+	/** Called after the world has been saved. The sequencer updates to the animated state. */
+	void OnPostSaveWorld(uint32 SaveFlags, class UWorld* World, bool bSuccess);
+
 private:
 	TMap< TWeakObjectPtr<UMovieSceneSection>, TSharedRef<FMovieSceneInstance> > MovieSceneSectionToInstanceMap;
 
diff --git a/Engine/Source/Editor/Sequencer/Private/SequencerAssetEditor.cpp b/Engine/Source/Editor/Sequencer/Private/SequencerAssetEditor.cpp
--- a/Engine/Source/Editor/Sequencer/Private/SequencerAssetEditor.cpp
+++ b/Engine/Source/Editor/Sequencer/Private/SequencerAssetEditor.cpp
@@ -112,6 +112,7 @@ FSequencerAssetEditor::~FSequencerAssetEditor()
 {
 	DetachTransportControlsFromViewports();
 
+	Sequencer->OnClose();
 
 	// Unregister delegates
 	if( FModuleManager::Get().IsModuleLoaded( TEXT( "LevelEditor" ) ) )
diff --git a/Engine/Source/Runtime/MovieSceneCore/Private/MovieSceneInstance.cpp b/Engine/Source/Runtime/MovieSceneCore/Private/MovieSceneInstance.cpp
--- a/Engine/Source/Runtime/MovieSceneCore/Private/MovieSceneInstance.cpp
+++ b/Engine/Source/Runtime/MovieSceneCore/Private/MovieSceneInstance.cpp
@@ -8,6 +8,34 @@ FMovieSceneInstance::FMovieSceneInstance( UMovieScene& InMovieScene )
 {
 }
 
+void FMovieSceneInstance::SaveState()
+{
+	TMap<FGuid, FMovieSceneObjectBindingInstance>::TIterator ObjectIt = ObjectBindingInstances.CreateIterator();
+	for (; ObjectIt; ++ObjectIt)
+	{
+		FMovieSceneObjectBindingInstance& ObjectBindingInstance = ObjectIt.Value();
+
+		for (FMovieSceneInstanceMap::TIterator It = ObjectBindingInstance.TrackInstances.CreateIterator(); It; ++It)
+		{
+			It.Value()->SaveState(ObjectBindingInstance.RuntimeObjects);
+		}
+	}
+}
+
+void FMovieSceneInstance::RestoreState()
+{
+	TMap<FGuid, FMovieSceneObjectBindingInstance>::TIterator ObjectIt = ObjectBindingInstances.CreateIterator();
+	for (; ObjectIt; ++ObjectIt)
+	{
+		FMovieSceneObjectBindingInstance& ObjectBindingInstance = ObjectIt.Value();
+
+		for (FMovieSceneInstanceMap::TIterator It = ObjectBindingInstance.TrackInstances.CreateIterator(); It; ++It)
+		{
+			It.Value()->RestoreState(ObjectBindingInstance.RuntimeObjects);
+		}
+	}
+}
+
 void FMovieSceneInstance::Update( float Position, float LastPosition, class IMovieScenePlayer& Player )
 {
 	// Update each type
@@ -92,7 +120,7 @@ void FMovieSceneInstance::RefreshInstanceMap( const TArray<UMovieSceneTrack*>& T
 		{
 			// The track does not have an instance, create one
 			Instance = Track->CreateInstance();
-
+			Instance->SaveState(RuntimeObjects);
 			Instance->RefreshInstance( RuntimeObjects, Player );
 
 			TrackInstances.Add( Track, Instance );
diff --git a/Engine/Source/Runtime/MovieSceneCore/Public/IMovieSceneTrackInstance.h b/Engine/Source/Runtime/MovieSceneCore/Public/IMovieSceneTrackInstance.h
--- a/Engine/Source/Runtime/MovieSceneCore/Public/IMovieSceneTrackInstance.h
+++ b/Engine/Source/Runtime/MovieSceneCore/Public/IMovieSceneTrackInstance.h
@@ -10,6 +10,16 @@ class IMovieSceneTrackInstance
 public:
 	virtual ~IMovieSceneTrackInstance() {}
 
+	/**
+	* Save state of objects that this instance will be editing.
+	*/
+	virtual void SaveState ( const TArray<UObject*>& RuntimeObjects ) {}
+
+	/**
+	* Restore state of objects that this instance edited.
+	*/
+	virtual void RestoreState ( const TArray<UObject*>& RuntimeObjects ) {}
+
 	/**
 	 * Main update function for track instances.  Called in game and in editor when we update a moviescene.
 	 *
diff --git a/Engine/Source/Runtime/MovieSceneCore/Public/MovieSceneInstance.h b/Engine/Source/Runtime/MovieSceneCore/Public/MovieSceneInstance.h
--- a/Engine/Source/Runtime/MovieSceneCore/Public/MovieSceneInstance.h
+++ b/Engine/Source/Runtime/MovieSceneCore/Public/MovieSceneInstance.h
@@ -10,6 +10,16 @@ class FMovieSceneInstance : public TSharedFromThis<FMovieSceneInstance>
 public:
 	MOVIESCENECORE_API FMovieSceneInstance( UMovieScene& InMovieScene );
 
+	/**
+	* Save state of the objects that this movie scene controls.
+	*/
+	MOVIESCENECORE_API void SaveState ();
+
+	/**
+	* Restore state of the objects that this movie scene controls.
+	*/
+	MOVIESCENECORE_API void RestoreState ();
+
 	/**
 	 * Updates this movie scene
 	 *
diff --git a/Engine/Source/Runtime/MovieSceneCoreTypes/Private/Instances/MovieScene3DTransformTrackInstance.cpp b/Engine/Source/Runtime/MovieSceneCoreTypes/Private/Instances/MovieScene3DTransformTrackInstance.cpp
--- a/Engine/Source/Runtime/MovieSceneCoreTypes/Private/Instances/MovieScene3DTransformTrackInstance.cpp
+++ b/Engine/Source/Runtime/MovieSceneCoreTypes/Private/Instances/MovieScene3DTransformTrackInstance.cpp
@@ -9,6 +9,57 @@ FMovieScene3DTransformTrackInstance::FMovieScene3DTransformTrackInstance( UMovie
 	TransformTrack = &InTransformTrack;
 }
 
+namespace
+{
+
+USceneComponent* SceneComponentFromRuntimeObject(UObject* Object)
+{
+	AActor* Actor = Cast<AActor>(Object);
+
+	USceneComponent* SceneComponent = NULL;
+	if (Actor && Actor->GetRootComponent())
+	{
+		// If there is an actor, modify its root component
+		SceneComponent = Actor->GetRootComponent();
+	}
+	else
+	{
+		// No actor was found.  Attempt to get the object as a component in the case that we are editing them directly.
+		SceneComponent = Cast<USceneComponent>(Object);
+	}
+	return SceneComponent;
+}
+
+}
+
+void FMovieScene3DTransformTrackInstance::SaveState(const TArray<UObject*>& RuntimeObjects)
+{
+	for (int32 ObjIndex = 0; ObjIndex < RuntimeObjects.Num(); ++ObjIndex)
+	{
+		USceneComponent* SceneComponent = SceneComponentFromRuntimeObject(RuntimeObjects[ObjIndex]);
+		if (SceneComponent != NULL)
+		{
+			InitTransformMap.Add(RuntimeObjects[ObjIndex], SceneComponent->GetRelativeTransform());
+		}
+	}
+}
+
+void FMovieScene3DTransformTrackInstance::RestoreState(const TArray<UObject*>& RuntimeObjects)
+{
+	for (int32 ObjIndex = 0; ObjIndex < RuntimeObjects.Num(); ++ObjIndex)
+	{
+		USceneComponent* SceneComponent = SceneComponentFromRuntimeObject(RuntimeObjects[ObjIndex]);
+		if (SceneComponent != NULL)
+		{
+			FTransform *Transform = InitTransformMap.Find(RuntimeObjects[ObjIndex]);
+			if (Transform != NULL)
+			{
+				SceneComponent->SetRelativeTransform(*Transform);
+			}
+		}
+	}
+}
+
 void FMovieScene3DTransformTrackInstance::Update( float Position, float LastPosition, const TArray<UObject*>& RuntimeObjects, class IMovieScenePlayer& Player ) 
 {
 	FVector Translation;
@@ -22,30 +73,19 @@ void FMovieScene3DTransformTrackInstance::Update( float Position, float LastPosi
 	{
 		for( int32 ObjIndex = 0; ObjIndex < RuntimeObjects.Num(); ++ObjIndex )
 		{
-			UObject* Object = RuntimeObjects[ObjIndex];
+			USceneComponent* SceneComponent = SceneComponentFromRuntimeObject(RuntimeObjects[ObjIndex]);
 
-			AActor* Actor = Cast<AActor>( Object );
-
-			USceneComponent* SceneComponent = NULL;
-			if( Actor && Actor->GetRootComponent() )
-			{
-				// If there is an actor, modify its root component
-				SceneComponent = Actor->GetRootComponent();
-			}
-			else
+			if (SceneComponent != NULL)
 			{
-				// No actor was found.  Attempt to get the object as a component in the case that we are editing them directly.
-				SceneComponent = Cast<USceneComponent>( Object );
-			}
+				// Set the relative translation and rotation.  Note they are set once instead of individually to avoid a redundant component transform update.
+				SceneComponent->SetRelativeLocationAndRotation(
+					bHasTranslationKeys ? Translation : SceneComponent->RelativeLocation,
+					bHasRotationKeys ? Rotation : SceneComponent->RelativeRotation );
 
-			// Set the relative translation and rotation.  Note they are set once instead of individually to avoid a redundant component transform update.
-			SceneComponent->SetRelativeLocationAndRotation(
-				bHasTranslationKeys ? Translation : SceneComponent->RelativeLocation,
-				bHasRotationKeys ? Rotation : SceneComponent->RelativeRotation );
-
-			if( bHasScaleKeys )
-			{
-				SceneComponent->SetRelativeScale3D( Scale );
+				if( bHasScaleKeys )
+				{
+					SceneComponent->SetRelativeScale3D( Scale );
+				}
 			}
 		}
 	}
diff --git a/Engine/Source/Runtime/MovieSceneCoreTypes/Private/Instances/MovieScene3DTransformTrackInstance.h b/Engine/Source/Runtime/MovieSceneCoreTypes/Private/Instances/MovieScene3DTransformTrackInstance.h
--- a/Engine/Source/Runtime/MovieSceneCoreTypes/Private/Instances/MovieScene3DTransformTrackInstance.h
+++ b/Engine/Source/Runtime/MovieSceneCoreTypes/Private/Instances/MovieScene3DTransformTrackInstance.h
@@ -11,9 +11,14 @@ public:
 	FMovieScene3DTransformTrackInstance( UMovieScene3DTransformTrack& InTransformTrack );
 
 	/** IMovieSceneTrackInstance interface */
+	virtual void SaveState (const TArray<UObject*>& RuntimeObjects) override;
+	virtual void RestoreState (const TArray<UObject*>& RuntimeObjects) override;
 	virtual void Update( float Position, float LastPosition, const TArray<UObject*>& RuntimeObjects, class IMovieScenePlayer& Player ) override;
 	virtual void RefreshInstance( const TArray<UObject*>& RuntimeObjects, class IMovieScenePlayer& Player ) override {}
 private:
 	/** Track that is being instanced */
 	UMovieScene3DTransformTrack* TransformTrack;
+
+	/** Map from object to initial state */
+	TMap< TWeakObjectPtr<UObject>, FTransform > InitTransformMap;
 };
diff --git a/Engine/Source/Runtime/MovieSceneCoreTypes/Private/Instances/MovieSceneBoolTrackInstance.cpp b/Engine/Source/Runtime/MovieSceneCoreTypes/Private/Instances/MovieSceneBoolTrackInstance.cpp
--- a/Engine/Source/Runtime/MovieSceneCoreTypes/Private/Instances/MovieSceneBoolTrackInstance.cpp
+++ b/Engine/Source/Runtime/MovieSceneCoreTypes/Private/Instances/MovieSceneBoolTrackInstance.cpp
@@ -11,6 +11,29 @@ FMovieSceneBoolTrackInstance::FMovieSceneBoolTrackInstance( UMovieSceneBoolTrack
 	PropertyBindings = MakeShareable( new FTrackInstancePropertyBindings( BoolTrack->GetPropertyName(), BoolTrack->GetPropertyPath() ) );
 }
 
+void FMovieSceneBoolTrackInstance::SaveState(const TArray<UObject*>& RuntimeObjects)
+{
+	for( UObject* Object : RuntimeObjects )
+	{
+		bool BoolValue = PropertyBindings->GetCurrentValue<bool>(Object);
+		InitBoolMap.Add(Object, BoolValue);
+	}
+}
+
+void FMovieSceneBoolTrackInstance::RestoreState(const TArray<UObject*>& RuntimeObjects)
+{
+	for( UObject* Object : RuntimeObjects )
+	{
+		bool *BoolValue = InitBoolMap.Find(Object);
+		if (BoolValue != NULL)
+		{
+			PropertyBindings->CallFunction(Object, BoolValue);
+		}
+	}
+
+	PropertyBindings->UpdateBindings( RuntimeObjects );
+}
+
 void FMovieSceneBoolTrackInstance::Update( float Position, float LastPosition, const TArray<UObject*>& RuntimeObjects, class IMovieScenePlayer& Player ) 
 {
 	bool BoolValue = false;
diff --git a/Engine/Source/Runtime/MovieSceneCoreTypes/Private/Instances/MovieSceneBoolTrackInstance.h b/Engine/Source/Runtime/MovieSceneCoreTypes/Private/Instances/MovieSceneBoolTrackInstance.h
--- a/Engine/Source/Runtime/MovieSceneCoreTypes/Private/Instances/MovieSceneBoolTrackInstance.h
+++ b/Engine/Source/Runtime/MovieSceneCoreTypes/Private/Instances/MovieSceneBoolTrackInstance.h
@@ -11,6 +11,8 @@ public:
 	FMovieSceneBoolTrackInstance( UMovieSceneBoolTrack& InBoolTrack );
 
 	/** IMovieSceneTrackInstance interface */
+	virtual void SaveState (const TArray<UObject*>& RuntimeObjects) override;
+	virtual void RestoreState (const TArray<UObject*>& RuntimeObjects) override;
 	virtual void Update( float Position, float LastPosition, const TArray<UObject*>& RuntimeObjects, class IMovieScenePlayer& Player ) override;
 	virtual void RefreshInstance( const TArray<UObject*>& RuntimeObjects, class IMovieScenePlayer& Player ) override;
 private:
@@ -18,4 +20,6 @@ private:
 	UMovieSceneBoolTrack* BoolTrack;
 	/** Runtime property bindings */
 	TSharedPtr<class FTrackInstancePropertyBindings> PropertyBindings;
+	/** Map from object to initial state */
+	TMap< TWeakObjectPtr<UObject>, bool > InitBoolMap;
 };
diff --git a/Engine/Source/Runtime/MovieSceneCoreTypes/Private/Instances/MovieSceneByteTrackInstance.cpp b/Engine/Source/Runtime/MovieSceneCoreTypes/Private/Instances/MovieSceneByteTrackInstance.cpp
--- a/Engine/Source/Runtime/MovieSceneCoreTypes/Private/Instances/MovieSceneByteTrackInstance.cpp
+++ b/Engine/Source/Runtime/MovieSceneCoreTypes/Private/Instances/MovieSceneByteTrackInstance.cpp
@@ -11,6 +11,29 @@ FMovieSceneByteTrackInstance::FMovieSceneByteTrackInstance( UMovieSceneByteTrack
 	PropertyBindings = MakeShareable( new FTrackInstancePropertyBindings( ByteTrack->GetPropertyName(), ByteTrack->GetPropertyPath() ) );
 }
 
+void FMovieSceneByteTrackInstance::SaveState(const TArray<UObject*>& RuntimeObjects)
+{
+	for( UObject* Object : RuntimeObjects )
+	{
+		uint8 ByteValue = PropertyBindings->GetCurrentValue<uint8>(Object);
+		InitByteMap.Add(Object, ByteValue);
+	}
+}
+
+void FMovieSceneByteTrackInstance::RestoreState(const TArray<UObject*>& RuntimeObjects)
+{
+	for( UObject* Object : RuntimeObjects )
+	{
+		uint8 *ByteValue = InitByteMap.Find(Object);
+		if (ByteValue != NULL)
+		{
+			PropertyBindings->CallFunction(Object, ByteValue);
+		}
+	}
+
+	PropertyBindings->UpdateBindings( RuntimeObjects );
+}
+
 void FMovieSceneByteTrackInstance::Update( float Position, float LastPosition, const TArray<UObject*>& RuntimeObjects, class IMovieScenePlayer& Player ) 
 {
 	uint8 ByteValue = 0;
diff --git a/Engine/Source/Runtime/MovieSceneCoreTypes/Private/Instances/MovieSceneByteTrackInstance.h b/Engine/Source/Runtime/MovieSceneCoreTypes/Private/Instances/MovieSceneByteTrackInstance.h
--- a/Engine/Source/Runtime/MovieSceneCoreTypes/Private/Instances/MovieSceneByteTrackInstance.h
+++ b/Engine/Source/Runtime/MovieSceneCoreTypes/Private/Instances/MovieSceneByteTrackInstance.h
@@ -11,6 +11,8 @@ public:
 	FMovieSceneByteTrackInstance( UMovieSceneByteTrack& InByteTrack );
 
 	/** IMovieSceneTrackInstance interface */
+	virtual void SaveState (const TArray<UObject*>& RuntimeObjects) override;
+	virtual void RestoreState (const TArray<UObject*>& RuntimeObjects) override;
 	virtual void Update( float Position, float LastPosition, const TArray<UObject*>& RuntimeObjects, class IMovieScenePlayer& Player ) override;
 	virtual void RefreshInstance( const TArray<UObject*>& RuntimeObjects, class IMovieScenePlayer& Player ) override;
 private:
@@ -18,4 +20,6 @@ private:
 	UMovieSceneByteTrack* ByteTrack;
 	/** Runtime property bindings */
 	TSharedPtr<class FTrackInstancePropertyBindings> PropertyBindings;
+	/** Map from object to initial state */
+	TMap< TWeakObjectPtr<UObject>, uint8 > InitByteMap;
 };
diff --git a/Engine/Source/Runtime/MovieSceneCoreTypes/Private/Instances/MovieSceneColorTrackInstance.cpp b/Engine/Source/Runtime/MovieSceneCoreTypes/Private/Instances/MovieSceneColorTrackInstance.cpp
--- a/Engine/Source/Runtime/MovieSceneCoreTypes/Private/Instances/MovieSceneColorTrackInstance.cpp
+++ b/Engine/Source/Runtime/MovieSceneCoreTypes/Private/Instances/MovieSceneColorTrackInstance.cpp
@@ -13,6 +13,48 @@ FMovieSceneColorTrackInstance::FMovieSceneColorTrackInstance( UMovieSceneColorTr
 	PropertyBindings = MakeShareable( new FTrackInstancePropertyBindings( ColorTrack->GetPropertyName(), ColorTrack->GetPropertyPath() ) );
 }
 
+void FMovieSceneColorTrackInstance::SaveState(const TArray<UObject*>& RuntimeObjects)
+{
+	for( UObject* Object : RuntimeObjects )
+	{
+		if( ColorTrack->IsSlateColor() )
+		{
+			FSlateColor ColorValue = PropertyBindings->GetCurrentValue<FSlateColor>(Object);
+			InitSlateColorMap.Add(Object, ColorValue);
+		}
+		else
+		{
+			FLinearColor ColorValue = PropertyBindings->GetCurrentValue<FLinearColor>(Object);
+			InitLinearColorMap.Add(Object, ColorValue);
+		}
+	}
+}
+
+void FMovieSceneColorTrackInstance::RestoreState(const TArray<UObject*>& RuntimeObjects)
+{
+	for( UObject* Object : RuntimeObjects )
+	{
+		if( ColorTrack->IsSlateColor() )
+		{
+			FSlateColor *ColorValue = InitSlateColorMap.Find(Object);
+			if (ColorValue != NULL)
+			{
+				PropertyBindings->CallFunction(Object, ColorValue);
+			}
+		}
+		else
+		{
+			FLinearColor *ColorValue = InitLinearColorMap.Find(Object);
+			if (ColorValue != NULL)
+			{
+				PropertyBindings->CallFunction(Object, ColorValue);
+			}
+		}
+	}
+
+	PropertyBindings->UpdateBindings( RuntimeObjects );
+}
+
 void FMovieSceneColorTrackInstance::Update( float Position, float LastPosition, const TArray<UObject*>& RuntimeObjects, class IMovieScenePlayer& Player ) 
 {
 	for(UObject* Object : RuntimeObjects)
diff --git a/Engine/Source/Runtime/MovieSceneCoreTypes/Private/Instances/MovieSceneColorTrackInstance.h b/Engine/Source/Runtime/MovieSceneCoreTypes/Private/Instances/MovieSceneColorTrackInstance.h
--- a/Engine/Source/Runtime/MovieSceneCoreTypes/Private/Instances/MovieSceneColorTrackInstance.h
+++ b/Engine/Source/Runtime/MovieSceneCoreTypes/Private/Instances/MovieSceneColorTrackInstance.h
@@ -11,6 +11,8 @@ public:
 	FMovieSceneColorTrackInstance( UMovieSceneColorTrack& InColorTrack );
 
 	/** IMovieSceneTrackInstance interface */
+	virtual void SaveState (const TArray<UObject*>& RuntimeObjects) override;
+	virtual void RestoreState (const TArray<UObject*>& RuntimeObjects) override;
 	virtual void Update( float Position, float LastPosition, const TArray<UObject*>& RuntimeObjects, class IMovieScenePlayer& Player ) override;
 	virtual void RefreshInstance( const TArray<UObject*>& RuntimeObjects, class IMovieScenePlayer& Player ) override;
 private:
@@ -18,4 +20,7 @@ private:
 	UMovieSceneColorTrack* ColorTrack;
 	/** Runtime property bindings */
 	TSharedPtr<class FTrackInstancePropertyBindings> PropertyBindings;
+	/** Map from object to initial state */
+	TMap< TWeakObjectPtr<UObject>, FSlateColor > InitSlateColorMap;
+	TMap< TWeakObjectPtr<UObject>, FLinearColor > InitLinearColorMap;
 };
diff --git a/Engine/Source/Runtime/MovieSceneCoreTypes/Private/Instances/MovieSceneFloatTrackInstance.cpp b/Engine/Source/Runtime/MovieSceneCoreTypes/Private/Instances/MovieSceneFloatTrackInstance.cpp
--- a/Engine/Source/Runtime/MovieSceneCoreTypes/Private/Instances/MovieSceneFloatTrackInstance.cpp
+++ b/Engine/Source/Runtime/MovieSceneCoreTypes/Private/Instances/MovieSceneFloatTrackInstance.cpp
@@ -11,6 +11,29 @@ FMovieSceneFloatTrackInstance::FMovieSceneFloatTrackInstance( UMovieSceneFloatTr
 	PropertyBindings = MakeShareable( new FTrackInstancePropertyBindings( FloatTrack->GetPropertyName(), FloatTrack->GetPropertyPath() ) );
 }
 
+void FMovieSceneFloatTrackInstance::SaveState(const TArray<UObject*>& RuntimeObjects)
+{
+	for( UObject* Object : RuntimeObjects )
+	{
+		float FloatValue = PropertyBindings->GetCurrentValue<float>(Object);
+		InitFloatMap.Add(Object, FloatValue);
+	}
+}
+
+void FMovieSceneFloatTrackInstance::RestoreState(const TArray<UObject*>& RuntimeObjects)
+{
+	for( UObject* Object : RuntimeObjects )
+	{
+		float *FloatValue = InitFloatMap.Find(Object);
+		if (FloatValue != NULL)
+		{
+			PropertyBindings->CallFunction(Object, FloatValue);
+		}
+	}
+
+	PropertyBindings->UpdateBindings( RuntimeObjects );
+}
+
 void FMovieSceneFloatTrackInstance::Update( float Position, float LastPosition, const TArray<UObject*>& RuntimeObjects, class IMovieScenePlayer& Player ) 
 {
 	float FloatValue = 0.0f;
diff --git a/Engine/Source/Runtime/MovieSceneCoreTypes/Private/Instances/MovieSceneFloatTrackInstance.h b/Engine/Source/Runtime/MovieSceneCoreTypes/Private/Instances/MovieSceneFloatTrackInstance.h
--- a/Engine/Source/Runtime/MovieSceneCoreTypes/Private/Instances/MovieSceneFloatTrackInstance.h
+++ b/Engine/Source/Runtime/MovieSceneCoreTypes/Private/Instances/MovieSceneFloatTrackInstance.h
@@ -11,6 +11,8 @@ public:
 	FMovieSceneFloatTrackInstance( UMovieSceneFloatTrack& InFloatTrack );
 
 	/** IMovieSceneTrackInstance interface */
+	virtual void SaveState (const TArray<UObject*>& RuntimeObjects) override;
+	virtual void RestoreState (const TArray<UObject*>& RuntimeObjects) override;
 	virtual void Update( float Position, float LastPosition, const TArray<UObject*>& RuntimeObjects, class IMovieScenePlayer& Player ) override;
 	virtual void RefreshInstance( const TArray<UObject*>& RuntimeObjects, class IMovieScenePlayer& Player ) override;
 private:
@@ -18,4 +20,6 @@ private:
 	UMovieSceneFloatTrack* FloatTrack;
 	/** Runtime property bindings */
 	TSharedPtr<class FTrackInstancePropertyBindings> PropertyBindings;
+	/** Map from object to initial state */
+	TMap< TWeakObjectPtr<UObject>, float > InitFloatMap;
 };
diff --git a/Engine/Source/Runtime/MovieSceneCoreTypes/Private/Instances/MovieSceneVectorTrackInstance.cpp b/Engine/Source/Runtime/MovieSceneCoreTypes/Private/Instances/MovieSceneVectorTrackInstance.cpp
--- a/Engine/Source/Runtime/MovieSceneCoreTypes/Private/Instances/MovieSceneVectorTrackInstance.cpp
+++ b/Engine/Source/Runtime/MovieSceneCoreTypes/Private/Instances/MovieSceneVectorTrackInstance.cpp
@@ -10,6 +10,75 @@ FMovieSceneVectorTrackInstance::FMovieSceneVectorTrackInstance( UMovieSceneVecto
 	PropertyBindings = MakeShareable( new FTrackInstancePropertyBindings( VectorTrack->GetPropertyName(), VectorTrack->GetPropertyPath() ) );
 }
 
+void FMovieSceneVectorTrackInstance::SaveState(const TArray<UObject*>& RuntimeObjects)
+{
+	int32 NumChannelsUsed = VectorTrack->GetNumChannelsUsed();
+	for( UObject* Object : RuntimeObjects )
+	{
+		switch( NumChannelsUsed )
+		{
+			case 2:
+			{
+				FVector2D VectorValue = PropertyBindings->GetCurrentValue<FVector2D>(Object);
+				InitVector2DMap.Add(Object, VectorValue);
+				break;
+			}
+			case 3:
+			{
+				FVector VectorValue = PropertyBindings->GetCurrentValue<FVector>(Object);
+				InitVector3DMap.Add(Object, VectorValue);
+				break;
+			}
+			case 4:
+			{
+				FVector4 VectorValue = PropertyBindings->GetCurrentValue<FVector4>(Object);
+				InitVector4DMap.Add(Object, VectorValue);
+				break;
+			}
+		}
+	}
+}
+
+void FMovieSceneVectorTrackInstance::RestoreState(const TArray<UObject*>& RuntimeObjects)
+{
+	int32 NumChannelsUsed = VectorTrack->GetNumChannelsUsed();
+	for( UObject* Object : RuntimeObjects )
+	{
+		switch( NumChannelsUsed )
+		{
+			case 2:
+			{
+				FVector2D *VectorValue = InitVector2DMap.Find(Object);
+				if (VectorValue != NULL)
+				{
+					PropertyBindings->CallFunction(Object, VectorValue);
+				}
+				break;
+			}
+			case 3:
+			{
+				FVector *VectorValue = InitVector3DMap.Find(Object);
+				if (VectorValue != NULL)
+				{
+					PropertyBindings->CallFunction(Object, VectorValue);
+				}
+				break;
+			}
+			case 4:
+			{
+				FVector4 *VectorValue = InitVector4DMap.Find(Object);
+				if (VectorValue != NULL)
+				{
+					PropertyBindings->CallFunction(Object, VectorValue);
+				}
+				break;
+			}
+		}
+	}
+
+	PropertyBindings->UpdateBindings( RuntimeObjects );
+}
+
 void FMovieSceneVectorTrackInstance::Update( float Position, float LastPosition, const TArray<UObject*>& RuntimeObjects, class IMovieScenePlayer& Player ) 
 {
 	FVector4 Vector;
diff --git a/Engine/Source/Runtime/MovieSceneCoreTypes/Private/Instances/MovieSceneVectorTrackInstance.h b/Engine/Source/Runtime/MovieSceneCoreTypes/Private/Instances/MovieSceneVectorTrackInstance.h
--- a/Engine/Source/Runtime/MovieSceneCoreTypes/Private/Instances/MovieSceneVectorTrackInstance.h
+++ b/Engine/Source/Runtime/MovieSceneCoreTypes/Private/Instances/MovieSceneVectorTrackInstance.h
@@ -11,6 +11,8 @@ public:
 	FMovieSceneVectorTrackInstance( UMovieSceneVectorTrack& InVectorTrack );
 
 	/** IMovieSceneTrackInstance interface */
+	virtual void SaveState (const TArray<UObject*>& RuntimeObjects) override;
+	virtual void RestoreState (const TArray<UObject*>& RuntimeObjects) override;
 	virtual void Update( float Position, float LastPosition, const TArray<UObject*>& RuntimeObjects, class IMovieScenePlayer& Player ) override;
 	virtual void RefreshInstance( const TArray<UObject*>& RuntimeObjects, class IMovieScenePlayer& Player ) override;
 private:
@@ -18,4 +20,8 @@ private:
 	UMovieSceneVectorTrack* VectorTrack;
 	/** Runtime property bindings */
 	TSharedPtr<class FTrackInstancePropertyBindings> PropertyBindings;
+	/** Map from object to initial state */
+	TMap< TWeakObjectPtr<UObject>, FVector2D > InitVector2DMap;
+	TMap< TWeakObjectPtr<UObject>, FVector > InitVector3DMap;
+	TMap< TWeakObjectPtr<UObject>, FVector4 > InitVector4DMap;
 };
